Bink Audio :: 3/4 :: Encoder :: #pf 60aea2116905a6000136901e

Unreal side of the encoder. Mostly just a thunk to the SDK, with some interleaving and splitting for the Unreal streaming system.

diff --git a/Engine/Source/Developer/AudioFormatBink/AudioFormatBink.Build.cs b/Engine/Source/Developer/AudioFormatBink/AudioFormatBink.Build.cs
new file mode 100644
--- /dev/null
+++ b/Engine/Source/Developer/AudioFormatBink/AudioFormatBink.Build.cs
@@ -0,0 +1,32 @@
+// Copyright Epic Games, Inc. All Rights Reserved.
+
+using UnrealBuildTool;
+using System.IO;
+
+public class AudioFormatBink : ModuleRules
+{
+	public AudioFormatBink(ReadOnlyTargetRules Target) : base(Target)
+	{
+		PrivateDependencyModuleNames.AddRange(
+			new string[] {
+				"Core",
+				"Engine"
+			}
+		);
+
+		PublicSystemIncludePaths.Add(Path.Combine(Target.UEThirdPartySourceDirectory, "BinkAudio", "Include"));
+
+		if (Target.Platform == UnrealTargetPlatform.Win64)
+		{
+			PublicAdditionalLibraries.Add(Path.Combine(Target.UEThirdPartySourceDirectory, "BinkAudio", "Lib", "binka_ue_encode_win64_static.lib"));
+		}
+		if (Target.Platform == UnrealTargetPlatform.Linux)
+		{
+			PublicAdditionalLibraries.Add(Path.Combine(Target.UEThirdPartySourceDirectory, "BinkAudio", "Lib", "libbinka_ue_encode_lnx64_static.a"));
+		}
+		if (Target.Platform == UnrealTargetPlatform.Mac)
+		{
+			PublicAdditionalLibraries.Add(Path.Combine(Target.UEThirdPartySourceDirectory, "BinkAudio", "Lib", "libbinka_ue_encode_osx64_static.a"));
+		}
+	}
+}
diff --git a/Engine/Source/Developer/AudioFormatBink/Private/AudioFormatBink.cpp b/Engine/Source/Developer/AudioFormatBink/Private/AudioFormatBink.cpp
new file mode 100644
--- /dev/null
+++ b/Engine/Source/Developer/AudioFormatBink/Private/AudioFormatBink.cpp
@@ -0,0 +1,199 @@
+// Copyright Epic Games, Inc. All Rights Reserved.
+
+#include "Modules/ModuleManager.h"
+#include "Interfaces/IAudioFormat.h"
+#include "Interfaces/IAudioFormatModule.h"
+
+#include "binka_ue_file_header.h"
+#include "binka_ue_encode.h"
+
+static const FName NAME_BINKA(TEXT("BINKA"));
+
+/**
+ * IAudioFormat, audio compression abstraction
+**/
+class FAudioFormatBink : public IAudioFormat
+{
+	enum
+	{
+		/** Version for Bink Audio format, this becomes part of the DDC key. */
+		UE_AUDIO_BINK_VER = 1,
+	};
+
+public:
+
+	virtual uint16 GetVersion(FName Format) const override
+	{
+		check(Format == NAME_BINKA);
+		return UE_AUDIO_BINK_VER;
+	}
+
+	virtual void GetSupportedFormats(TArray<FName>& OutFormats) const override
+	{
+		OutFormats.Add(NAME_BINKA);
+	}
+
+	static void* BinkAlloc(size_t Bytes)
+	{
+		return FMemory::Malloc(Bytes, 16);
+	}
+	static void BinkFree(void* Ptr)
+	{
+		FMemory::Free(Ptr);
+	}
+
+	virtual bool Cook(FName InFormat, const TArray<uint8>& InSrcBuffer, FSoundQualityInfo& InQualityInfo, TArray<uint8>& OutCompressedDataStore) const override
+	{
+		check(InFormat == NAME_BINKA);
+
+		// Bink goes from 0 (best) to 9 (worst), but is basically unusable below like 4,
+		// so we map 0-100 to 4-0
+		uint8 CompressionLevel = 0 + (4 - (4*InQualityInfo.Quality) / 100); 
+		
+		void* CompressedData = 0;
+		uint32_t CompressedDataLen = 0;
+		UECompressBinkAudio((void*)InSrcBuffer.GetData(), InSrcBuffer.Num(), InQualityInfo.SampleRate, InQualityInfo.NumChannels, CompressionLevel, 1, BinkAlloc, BinkFree, &CompressedData, &CompressedDataLen);
+
+		// Create a buffer to store compressed data
+		OutCompressedDataStore.Empty();
+		OutCompressedDataStore.Append((uint8*)CompressedData, CompressedDataLen);
+		BinkFree(CompressedData);
+		return OutCompressedDataStore.Num() > 0;
+	}
+
+	virtual bool CookSurround(FName InFormat, const TArray<TArray<uint8> >& InSrcBuffers, FSoundQualityInfo& InQualityInfo, TArray<uint8>& OutCompressedDataStore) const override
+	{
+		check(InFormat == NAME_BINKA);
+
+		//
+		// CookSurround passes us a bunch of mono buffers, but bink audio wants a standard 
+		// interleaved buffer
+		//
+		TArray<uint8> InterleavedSrcBuffers;
+		InterleavedSrcBuffers.AddUninitialized(InSrcBuffers[0].Num() * InSrcBuffers.Num());
+		int16* Dest = (int16*)InterleavedSrcBuffers.GetData();
+
+		uint32 ChannelCount = InSrcBuffers.Num();
+		uint32 FrameCount = InSrcBuffers[0].Num() / sizeof(int16);
+		for (uint32 FrameIndex = 0; FrameIndex < FrameCount; FrameIndex++)
+		{
+			for (uint32 ChannelIndex = 0; ChannelIndex < ChannelCount; ChannelIndex++)
+			{
+				int16* Src = (int16*)InSrcBuffers[ChannelIndex].GetData();
+				Dest[FrameIndex * ChannelCount + ChannelIndex] = Src[FrameIndex];
+			}
+		}
+
+		return Cook(NAME_BINKA, InterleavedSrcBuffers, InQualityInfo, OutCompressedDataStore);
+	}
+
+	// AFAICT this function is never called.
+	virtual int32 Recompress(FName Format, const TArray<uint8>& SrcBuffer, FSoundQualityInfo& QualityInfo, TArray<uint8>& OutBuffer) const override
+	{
+		return 0;
+	}
+
+	virtual int32 GetMinimumSizeForInitialChunk(FName Format, const TArray<uint8>& SrcBuffer) const override
+	{
+		// We must have an initial chunk large enough for the header and the seek table, if present.
+		BinkAudioFileHeader const* Header = (BinkAudioFileHeader const*)SrcBuffer.GetData();
+		return sizeof(BinkAudioFileHeader) + Header->seek_table_entry_count * sizeof(uint16);
+	}
+
+	// Takes in a compressed file and splits it in to stream size chunks. AFAICT this is supposed to accumulate frames
+	// until the chunk size is reached, then spit out a block.
+	virtual bool SplitDataForStreaming(const TArray<uint8>& InSrcBuffer, TArray<TArray<uint8>>& OutBuffers, const int32 InMaxInitialChunkSize, const int32 InMaxChunkSize) const override
+	{
+		uint8 const* Source = InSrcBuffer.GetData();
+		uint32 SourceLen = InSrcBuffer.Num();
+		uint8 const* SourceEnd = Source + SourceLen;
+
+		uint8 const* ChunkStart = Source;
+		uint8 const* Current = Source;
+	
+		ensure(SourceLen > sizeof(BinkAudioFileHeader));
+
+		BinkAudioFileHeader const* Header = (BinkAudioFileHeader const*)Source;
+		Current += sizeof(BinkAudioFileHeader);
+
+		ensure(Header->tag == 'UEBA');
+
+		// The first frame is located past the seek table.
+		Current += sizeof(uint16) * Header->seek_table_entry_count;
+
+		// We must be on the first frame.
+		ensure( (*((uint16*)Current)) == BLOCK_HEADER_MAGIC);
+
+		int32 ChunkLimitBytes = InMaxInitialChunkSize;
+		for (;;)
+		{
+			if (Current >= SourceEnd)
+			{
+				// Done with the file.
+				check(Current == SourceEnd);
+				Current = SourceEnd;
+				break;
+			}
+
+			// Advance to next chunk
+			uint32 BlockSize;
+			if (BinkAudioBlockSize(Header->max_comp_space_needed, Current, (SourceEnd - Current), &BlockSize) == false)
+			{
+				// if this happens, then UE didn't actually give us the data we gave it, and we're
+				// off the edge of the map.
+				check(0);
+				return false;
+			}
+
+			if ((Current - ChunkStart) + BlockSize >= ChunkLimitBytes)
+			{
+				// can't add this chunk, emit.
+				TArray<uint8> Chunk(ChunkStart, Current - ChunkStart);
+				OutBuffers.Add(Chunk);
+				ChunkStart = Current;
+				ChunkLimitBytes = InMaxChunkSize;
+
+				// retry.
+				continue;
+			}
+
+			Current += BlockSize;
+		}
+
+		// emit any remainder chunks
+		if (Current - ChunkStart)
+		{
+			// emit this chunk
+			TArray<uint8> Chunk(ChunkStart, Current - ChunkStart);
+			OutBuffers.Add(Chunk);
+			ChunkStart = Current;
+		}
+
+		return true;
+	}
+};
+
+class FAudioPlatformBinkModule : public IAudioFormatModule
+{
+private:
+	FAudioFormatBink* BinkEncoder = 0;
+
+public:
+	virtual ~FAudioPlatformBinkModule() {}
+	
+	virtual IAudioFormat* GetAudioFormat()
+	{
+		return BinkEncoder;
+	}
+	virtual void StartupModule()
+	{
+		BinkEncoder = new FAudioFormatBink();
+	}
+	virtual void ShutdownModule()
+	{
+		delete BinkEncoder;
+		BinkEncoder = 0;
+	}
+};
+
+IMPLEMENT_MODULE( FAudioPlatformBinkModule, AudioFormatBink);
