Adding a warning so we remember to check Bink Audio encoding when the editor is up and running on apple silicon.

diff --git a/Engine/Source/Developer/AudioFormatBink/Private/AudioFormatBink.cpp b/Engine/Source/Developer/AudioFormatBink/Private/AudioFormatBink.cpp
--- a/Engine/Source/Developer/AudioFormatBink/Private/AudioFormatBink.cpp
+++ b/Engine/Source/Developer/AudioFormatBink/Private/AudioFormatBink.cpp
@@ -3,12 +3,15 @@
 #include "Modules/ModuleManager.h"
 #include "Interfaces/IAudioFormat.h"
 #include "Interfaces/IAudioFormatModule.h"
+#include "HAL/Platform.h"
 
 #include "binka_ue_file_header.h"
 #include "binka_ue_encode.h"
 
 static const FName NAME_BINKA(TEXT("BINKA"));
 
+DEFINE_LOG_CATEGORY_STATIC(LogAudioFormatBink, Display, All);
+
 /**
  * IAudioFormat, audio compression abstraction
 **/
@@ -51,6 +54,10 @@ public:
 		TRACE_CPUPROFILER_EVENT_SCOPE(FAudioFormatBink::Cook);
 		check(InFormat == NAME_BINKA);
 
+#if !PLATFORM_CPU_X86_FAMILY
+		UE_LOG(LogAudioFormatBink, Warning, TEXT("Bink Audio encoder has not been tested for consistency on non x86 platforms - cooks on this platform might not match cooks on x86/x64 for the same DDC key!"));
+#endif
+
 		// Bink goes from 0 (best) to 9 (worst), but is basically unusable below like 4,
 		// so we map 0-100 to 4-0
 		uint8 CompressionLevel = 0 + (4 - (4*InQualityInfo.Quality) / 100); 
