Fixup Bink Quality mapping logic based on community feedback

diff --git a/Engine/Source/Developer/AudioFormatBink/Private/AudioFormatBink.cpp b/Engine/Source/Developer/AudioFormatBink/Private/AudioFormatBink.cpp
--- a/Engine/Source/Developer/AudioFormatBink/Private/AudioFormatBink.cpp
+++ b/Engine/Source/Developer/AudioFormatBink/Private/AudioFormatBink.cpp
@@ -12,6 +12,22 @@ static const FName NAME_BINKA(TEXT("BINKA"));
 
 DEFINE_LOG_CATEGORY_STATIC(LogAudioFormatBink, Display, All);
 
+namespace AudioFormatBinkPrivate
+{
+	uint8 GetCompressionLevelFromQualityIndex(int32 InQualityIndex) 
+	{
+		// Bink goes from 0 (best) to 9 (worst), but is basically unusable below 4 
+		static const float BinkLowest = 4;
+		static const float BinkHighest = 0;
+
+		// Map Quality 1 (lowest) to 40 (highest).
+		static const float QualityLowest = 1;
+		static const float QualityHighest = 40;
+
+		return FMath::GetMappedRangeValueClamped(FVector2D(QualityLowest, QualityHighest), FVector2D(BinkLowest, BinkHighest), InQualityIndex);
+	}	
+}
+
 /**
  * IAudioFormat, audio compression abstraction
 **/
@@ -20,7 +36,7 @@ class FAudioFormatBink : public IAudioFormat
 	enum
 	{
 		/** Version for Bink Audio format, this becomes part of the DDC key. */
-		UE_AUDIO_BINK_VER = 2,
+		UE_AUDIO_BINK_VER = 3,
 	};
 
 public:
@@ -58,9 +74,7 @@ public:
 		UE_LOG(LogAudioFormatBink, Warning, TEXT("Bink Audio encoder has not been tested for consistency on non x86 platforms - cooks on this platform might not match cooks on x86/x64 for the same DDC key!"));
 #endif
 
-		// Bink goes from 0 (best) to 9 (worst), but is basically unusable below like 4,
-		// so we map 0-100 to 4-0
-		uint8 CompressionLevel = 0 + (4 - (4*InQualityInfo.Quality) / 100); 
+		uint8 CompressionLevel = AudioFormatBinkPrivate::GetCompressionLevelFromQualityIndex(InQualityInfo.Quality);
 		
 		void* CompressedData = 0;
 		uint32_t CompressedDataLen = 0;
