Adding logs for bink audio encode errors.

diff --git a/Engine/Source/Developer/AudioFormatBink/Private/AudioFormatBink.cpp b/Engine/Source/Developer/AudioFormatBink/Private/AudioFormatBink.cpp
--- a/Engine/Source/Developer/AudioFormatBink/Private/AudioFormatBink.cpp
+++ b/Engine/Source/Developer/AudioFormatBink/Private/AudioFormatBink.cpp
@@ -74,7 +74,25 @@ public:
 		
 		void* CompressedData = 0;
 		uint32_t CompressedDataLen = 0;
-		UECompressBinkAudio((void*)InSrcBuffer.GetData(), InSrcBuffer.Num(), InQualityInfo.SampleRate, InQualityInfo.NumChannels, CompressionLevel, 1, BinkAlloc, BinkFree, &CompressedData, &CompressedDataLen);
+		uint8_t BinkCompressError = UECompressBinkAudio((void*)InSrcBuffer.GetData(), InSrcBuffer.Num(), InQualityInfo.SampleRate, InQualityInfo.NumChannels, CompressionLevel, 1, BinkAlloc, BinkFree, &CompressedData, &CompressedDataLen);
+
+		const TCHAR* CompressErrorStr = nullptr;
+		switch (BinkCompressError)
+		{
+		case BINKA_COMPRESS_SUCCESS: break;
+		case BINKA_COMPRESS_ERROR_CHANS: CompressErrorStr = TEXT("Invalid channel count, max ") BINKA_MAX_CHANS_STR; break;
+		case BINKA_COMPRESS_ERROR_SAMPLES: CompressErrorStr = TEXT("No sample data provided"); break;
+		case BINKA_COMPRESS_ERROR_RATE: CompressErrorStr = TEXT("Invalid sample rate provided, min ") BINKA_MIN_RATE_STR TEXT(" max ") BINKA_MAX_RATE_STR; break;
+		case BINKA_COMPRESS_ERROR_QUALITY: CompressErrorStr = TEXT("Invalid quality provided, valid is 0-9"); break;
+		case BINKA_COMPRESS_ERROR_ALLOCATORS: CompressErrorStr = TEXT("No allocators provided!"); break;
+		case BINKA_COMPRESS_ERROR_OUTPUT: CompressErrorStr = TEXT("No output pointers provided!"); break;
+		}
+
+		if (CompressErrorStr != nullptr)
+		{
+			UE_LOG(LogAudioFormatBink, Warning, TEXT("Failed to encode bink audio: %s"), CompressErrorStr);
+			return false;
+		}
 
 		// Create a buffer to store compressed data
 		OutCompressedDataStore.Empty();
diff --git a/Engine/Source/Runtime/BinkAudioDecoder/SDK/BinkAudio/Include/binka_ue_encode.h b/Engine/Source/Runtime/BinkAudioDecoder/SDK/BinkAudio/Include/binka_ue_encode.h
--- a/Engine/Source/Runtime/BinkAudioDecoder/SDK/BinkAudio/Include/binka_ue_encode.h
+++ b/Engine/Source/Runtime/BinkAudioDecoder/SDK/BinkAudio/Include/binka_ue_encode.h
@@ -5,6 +5,19 @@
 typedef void* BAUECompressAllocFnType(uintptr_t ByteCount);
 typedef void BAUECompressFreeFnType(void* Ptr);
 
+// (NOTE! This was copied after integrating the latest source, so when updating make sure
+// it stays around!
+#define BINKA_COMPRESS_SUCCESS 0
+#define BINKA_COMPRESS_ERROR_CHANS 1
+#define BINKA_COMPRESS_ERROR_SAMPLES 2
+#define BINKA_COMPRESS_ERROR_RATE 3
+#define BINKA_COMPRESS_ERROR_QUALITY 4
+#define BINKA_COMPRESS_ERROR_ALLOCATORS 5
+#define BINKA_COMPRESS_ERROR_OUTPUT 6
+
+#define BINKA_MAX_CHANS_STR TEXT("16") // (8*2) MAX_STREAMS stereo
+#define BINKA_MIN_RATE_STR TEXT("11 khz")
+#define BINKA_MAX_RATE_STR TEXT("256 khz")
 
 //
 // Compresses a bink audio _file_. This is distinct from _blocks_, which is
