Bink Format; use duration/samplerate to set max seektable size.

diff --git a/Engine/Source/Developer/AudioFormatBink/Private/AudioFormatBink.cpp b/Engine/Source/Developer/AudioFormatBink/Private/AudioFormatBink.cpp
--- a/Engine/Source/Developer/AudioFormatBink/Private/AudioFormatBink.cpp
+++ b/Engine/Source/Developer/AudioFormatBink/Private/AudioFormatBink.cpp
@@ -14,6 +14,23 @@ DEFINE_LOG_CATEGORY_STATIC(LogAudioFormatBink, Display, All);
 
 namespace AudioFormatBinkPrivate
 {
+	// Taken from BinkAudioInfo.cpp 
+	static uint32 GetMaxFrameSizeSamples(uint32 SampleRate)
+	{
+		if (SampleRate >= 44100)
+		{
+			return 1920;
+		}
+		else if (SampleRate >= 22050)
+		{
+			return 960;
+		}
+		else
+		{
+			return 480;
+		}
+	}
+
 	uint8 GetCompressionLevelFromQualityIndex(int32 InQualityIndex) 
 	{
 		// Bink goes from 0 (best) to 9 (worst), but is basically unusable below 4 
@@ -30,6 +47,17 @@ namespace AudioFormatBinkPrivate
 		// Floor each value and clamp into range (as top lerp will be +1 over)
 		return FMath::Clamp(FMath::FloorToInt(BinkValue), BinkHighest, BinkLowest);
 	}	
+
+	static uint16 GetMaxSeekTableEntries(const FSoundQualityInfo& InQualityInfo)
+	{
+		const uint64 DurationFrames = InQualityInfo.SampleDataSize / (InQualityInfo.NumChannels * sizeof(int16));
+		const uint64 MaxEntries = DurationFrames / GetMaxFrameSizeSamples(InQualityInfo.SampleRate);
+				
+		static const uint16 BinkDefault = 4096;
+		static const uint16 BinkMax		= TNumericLimits<uint16>::Max();
+
+		return IntCastChecked<uint16>(FMath::Clamp<uint64>(MaxEntries, BinkDefault, BinkMax));
+	}
 }
 
 /**
@@ -40,7 +68,7 @@ class FAudioFormatBink : public IAudioFormat
 	enum
 	{
 		/** Version for Bink Audio format, this becomes part of the DDC key. */
-		UE_AUDIO_BINK_VER = 7,
+		UE_AUDIO_BINK_VER = 9,
 	};
 
 public:
@@ -74,11 +102,14 @@ public:
 		TRACE_CPUPROFILER_EVENT_SCOPE(FAudioFormatBink::Cook);
 		check(InFormat == NAME_BINKA);
 
-		uint8 CompressionLevel = AudioFormatBinkPrivate::GetCompressionLevelFromQualityIndex(InQualityInfo.Quality);
+		const uint8 CompressionLevel = AudioFormatBinkPrivate::GetCompressionLevelFromQualityIndex(InQualityInfo.Quality);
+		
+		// If we're going to embed the seek-table in the stream, use -1 to give the largest table we can produce.
+		const uint16 MaxSeektableSize = AudioFormatBinkPrivate::GetMaxSeekTableEntries(InQualityInfo);
 		
 		void* CompressedData = 0;
 		uint32_t CompressedDataLen = 0;
-		uint8_t BinkCompressError = UECompressBinkAudio((void*)InSrcBuffer.GetData(), InSrcBuffer.Num(), InQualityInfo.SampleRate, InQualityInfo.NumChannels, CompressionLevel, 1, 4096, BinkAlloc, BinkFree, &CompressedData, &CompressedDataLen);
+		uint8_t BinkCompressError = UECompressBinkAudio((void*)InSrcBuffer.GetData(), InSrcBuffer.Num(), InQualityInfo.SampleRate, InQualityInfo.NumChannels, CompressionLevel, 1, MaxSeektableSize, BinkAlloc, BinkFree, &CompressedData, &CompressedDataLen);
 
 		const TCHAR* CompressErrorStr = nullptr;
 		switch (BinkCompressError)
@@ -258,23 +289,6 @@ public:
 	{
 		return true; // Toggling this will require a version bump.
 	}
-	
-	// Taken from BinkAudioInfo.cpp 
-	static uint32 GetMaxFrameSizeSamples(uint32 SampleRate) 
-	{
-		if (SampleRate >= 44100)
-		{
-			return 1920;
-		}
-		else if (SampleRate >= 22050)
-		{
-			return 960;
-		}
-		else
-		{
-			return 480;
-		}
-	}
 
 	virtual bool ExtractSeekTableForStreaming(TArray<uint8>& InOutBuffer, IAudioFormat::FSeekTable& OutSeektable) const
 	{
@@ -299,7 +313,7 @@ public:
 		uint32 CurrentSeekOffset = ActualAudioOffset;
 		uint32 CurrentTimeOffset = 0;
 		
-		int32 SamplesPerEntry = Header->blocks_per_seek_table_entry * GetMaxFrameSizeSamples(Header->rate);
+		int32 SamplesPerEntry = Header->blocks_per_seek_table_entry * AudioFormatBinkPrivate::GetMaxFrameSizeSamples(Header->rate);
 		OutSeektable.Offsets.SetNum(Header->seek_table_entry_count);
 		OutSeektable.Times.SetNum(Header->seek_table_entry_count);
 
@@ -341,4 +355,4 @@ public:
 	}
 };
 
-IMPLEMENT_MODULE( FAudioPlatformBinkModule, AudioFormatBink);
+IMPLEMENT_MODULE( FAudioPlatformBinkModule, AudioFormatBink);
\ No newline at end of file
