Retroactive fix up from comments from review and const tidyup

diff --git a/Engine/Source/Developer/AudioFormatBink/Private/AudioFormatBink.cpp b/Engine/Source/Developer/AudioFormatBink/Private/AudioFormatBink.cpp
--- a/Engine/Source/Developer/AudioFormatBink/Private/AudioFormatBink.cpp
+++ b/Engine/Source/Developer/AudioFormatBink/Private/AudioFormatBink.cpp
@@ -15,7 +15,7 @@ DEFINE_LOG_CATEGORY_STATIC(LogAudioFormatBink, Display, All);
 namespace AudioFormatBinkPrivate
 {
 	// Taken from BinkAudioInfo.cpp 
-	static uint32 GetMaxFrameSizeSamples(uint32 SampleRate)
+	static uint32 GetMaxFrameSizeSamples(const uint32 SampleRate)
 	{
 		if (SampleRate >= 44100)
 		{
@@ -31,18 +31,18 @@ namespace AudioFormatBinkPrivate
 		}
 	}
 
-	uint8 GetCompressionLevelFromQualityIndex(int32 InQualityIndex) 
+	static uint8 GetCompressionLevelFromQualityIndex(const int32 InQualityIndex) 
 	{
 		// Bink goes from 0 (best) to 9 (worst), but is basically unusable below 4 
-		static const float BinkLowest = 4;
-		static const float BinkHighest = 0;
+		static constexpr float BinkLowest = 4;
+		static constexpr float BinkHighest = 0;
 
 		// Map Quality 1 (lowest) to 100 (highest).
-		static const float QualityLowest = 1;
-		static const float QualityHighest = 100;
+		static constexpr float QualityLowest = 1;
+		static constexpr float QualityHighest = 100;
 
 		// Map Quality into Bink Range. Note: +1 gives the Bink range 5 steps inclusive.
-		float BinkValue = FMath::GetMappedRangeValueClamped(FVector2D(QualityLowest, QualityHighest), FVector2D(BinkLowest + 1.f, BinkHighest), InQualityIndex);
+		const float BinkValue = FMath::GetMappedRangeValueClamped(FVector2D(QualityLowest, QualityHighest), FVector2D(BinkLowest + 1.f, BinkHighest), InQualityIndex);
 
 		// Floor each value and clamp into range (as top lerp will be +1 over)
 		return FMath::Clamp(FMath::FloorToInt(BinkValue), BinkHighest, BinkLowest);
@@ -53,8 +53,8 @@ namespace AudioFormatBinkPrivate
 		const uint64 DurationFrames = InQualityInfo.SampleDataSize / (InQualityInfo.NumChannels * sizeof(int16));
 		const uint64 MaxEntries = DurationFrames / GetMaxFrameSizeSamples(InQualityInfo.SampleRate);
 				
-		static const uint16 BinkDefault = 4096;
-		static const uint16 BinkMax		= TNumericLimits<uint16>::Max();
+		static constexpr uint16 BinkDefault = 4096;
+		static constexpr uint16 BinkMax		= TNumericLimits<uint16>::Max();
 
 		const uint64 Clamped = FMath::Clamp<uint64>(MaxEntries, BinkDefault, BinkMax);
 			
@@ -90,7 +90,7 @@ public:
 		OutFormats.Add(NAME_BINKA);
 	}
 
-	static void* BinkAlloc(size_t Bytes)
+	static void* BinkAlloc(const size_t Bytes)
 	{
 		return FMemory::Malloc(Bytes, 16);
 	}
@@ -288,12 +288,12 @@ public:
 		InOutBuffer.SetNum(CompressedDataLen);
 	}
 	
-	virtual bool RequiresStreamingSeekTable() const
+	virtual bool RequiresStreamingSeekTable() const override
 	{
 		return true; // Toggling this will require a version bump.
 	}
 
-	virtual bool ExtractSeekTableForStreaming(TArray<uint8>& InOutBuffer, IAudioFormat::FSeekTable& OutSeektable) const
+	virtual bool ExtractSeekTableForStreaming(TArray<uint8>& InOutBuffer, IAudioFormat::FSeekTable& OutSeektable) const override
 	{
 		// This should only be called if we require a streaming seek-table. 
 		if (!ensure(RequiresStreamingSeekTable()))
@@ -301,7 +301,7 @@ public:
 			return false;
 		}
 
-		BinkAudioFileHeader const* Header = (BinkAudioFileHeader const*)InOutBuffer.GetData();
+		BinkAudioFileHeader const* Header = reinterpret_cast<BinkAudioFileHeader const*>(InOutBuffer.GetData());
 		if (InOutBuffer.Num() < sizeof(BinkAudioFileHeader) || Header->tag != 'UEBA' || Header->seek_table_entry_count==0)
 		{
 			return false;
@@ -309,10 +309,10 @@ public:
 		
 		// The outer logic that manages the seek-table is unaware of how big the header can be, so for the sake of simplicity,
 		// offset for the size of the header in the seek-table entries, so we don't need to worry about adjusting for it later.
-		static const uint32 ActualAudioOffset = sizeof(BinkAudioFileHeader) + 0; // No entries as we're stripping it.
+		static constexpr uint32 ActualAudioOffset = sizeof(BinkAudioFileHeader) + 0; // No entries as we're stripping it.
 
 		// Decode and copy out the seek-table. (it's stored as deltas).
-		const uint16* EncodedSeekTable = (uint16*)(InOutBuffer.GetData() + sizeof(BinkAudioFileHeader));
+		const uint16* EncodedSeekTable = reinterpret_cast<uint16*>(InOutBuffer.GetData() + sizeof(BinkAudioFileHeader));
 		uint32 CurrentSeekOffset = ActualAudioOffset;
 		uint32 CurrentTimeOffset = 0;
 		
@@ -342,26 +342,26 @@ public:
 	}
 };
 
-class FAudioPlatformBinkModule : public IAudioFormatModule
+class FAudioPlatformBinkModule final : public IAudioFormatModule
 {
 private:
-	FAudioFormatBink* BinkEncoder = 0;
+	FAudioFormatBink* BinkEncoder = nullptr;
 
 public:
-	virtual ~FAudioPlatformBinkModule() {}
+	virtual ~FAudioPlatformBinkModule() override {}
 	
-	virtual IAudioFormat* GetAudioFormat()
+	virtual IAudioFormat* GetAudioFormat() override
 	{
 		return BinkEncoder;
 	}
-	virtual void StartupModule()
+	virtual void StartupModule() override
 	{
 		BinkEncoder = new FAudioFormatBink();
 	}
-	virtual void ShutdownModule()
+	virtual void ShutdownModule() override
 	{
 		delete BinkEncoder;
-		BinkEncoder = 0;
+		BinkEncoder = nullptr;
 	}
 };
 
diff --git a/Engine/Source/Runtime/BinkAudioDecoder/SDK/BinkAudio/Src/binka_ue_encode.cpp b/Engine/Source/Runtime/BinkAudioDecoder/SDK/BinkAudio/Src/binka_ue_encode.cpp
--- a/Engine/Source/Runtime/BinkAudioDecoder/SDK/BinkAudio/Src/binka_ue_encode.cpp
+++ b/Engine/Source/Runtime/BinkAudioDecoder/SDK/BinkAudio/Src/binka_ue_encode.cpp
@@ -617,18 +617,3 @@ extern "C" int LLVMFuzzerTestOneInput(const uint8_t *Data, size_t Size)
     return 0;
 }
 #endif
-
-/* @cdep pre
-    $NoAutoInclude
-    $Defaults
-    $requires(popmal.c)
-    $requires(binkace.c)
-    $requires(binka_ue_encode.cpp)
-    $requires(radfft.cpp)
-    $requires(varbits.c)
-    $requires(ranged_log.cpp)
-*/
-
-/* @cdep post
-    $Build
-*/
